[LiveLinkVCAM] Fix Unreal failing to respond to discovery beacons on Mac

diff --git a/Engine/Plugins/DiscoveryBeaconReceiver/Source/DiscoveryBeaconReceiver/Private/DiscoveryBeaconReceiver.cpp b/Engine/Plugins/DiscoveryBeaconReceiver/Source/DiscoveryBeaconReceiver/Private/DiscoveryBeaconReceiver.cpp
--- a/Engine/Plugins/DiscoveryBeaconReceiver/Source/DiscoveryBeaconReceiver/Private/DiscoveryBeaconReceiver.cpp
+++ b/Engine/Plugins/DiscoveryBeaconReceiver/Source/DiscoveryBeaconReceiver/Private/DiscoveryBeaconReceiver.cpp
@@ -145,6 +145,13 @@ void FDiscoveryBeaconReceiver::ReceiveBeaconMessages()
 		TSharedRef<FInternetAddr> Source = ISocketSubsystem::Get()->CreateInternetAddr();
 		if (Socket->RecvFrom(MessageData.GetData(), MessageData.Num(), NumRead, *Source) && NumRead > 0)
 		{
+			if (NumRead < MessageData.Num())
+			{
+				// Update container size to reflect the actual message (which only matters if it's smaller, since the message
+				// will fill the container otherwise).
+				// Don't allow shrinking since this could incur a move and we're just going to release the array anyway.
+				MessageData.SetNum(NumRead, EAllowShrinking::No);
+			}
 			HandleBeaconMessage(MessageData, Source);
 		}
 	}
